Rename/move NativePointerUpgradeTool files to new UnrealObjectPtrTool path/filenames.  File content changes to follow.

diff --git a/Engine/Source/Programs/NativePointerUpgradeTool/NativePointerUpgradeTool.Build.cs b/Engine/Source/Programs/NativePointerUpgradeTool/NativePointerUpgradeTool.Build.cs
deleted file mode 100644
diff --git a/Engine/Source/Programs/NativePointerUpgradeTool/NativePointerUpgradeTool.Target.cs b/Engine/Source/Programs/NativePointerUpgradeTool/NativePointerUpgradeTool.Target.cs
deleted file mode 100644
diff --git a/Engine/Source/Programs/NativePointerUpgradeTool/Private/NativePointerUpgradeTool.cpp b/Engine/Source/Programs/NativePointerUpgradeTool/Private/NativePointerUpgradeTool.cpp
deleted file mode 100644
diff --git a/Engine/Source/Programs/UnrealObjectPtrTool/Private/UnrealObjectPtrTool.cpp b/Engine/Source/Programs/UnrealObjectPtrTool/Private/UnrealObjectPtrTool.cpp
new file mode 100644
--- /dev/null
+++ b/Engine/Source/Programs/UnrealObjectPtrTool/Private/UnrealObjectPtrTool.cpp
@@ -0,0 +1,623 @@
+// Copyright Epic Games, Inc. All Rights Reserved.
+
+
+#include "CoreMinimal.h"
+#include "RequiredProgramMainCPPInclude.h"
+#include "CoreTypes.h"
+#include "Containers/Array.h"
+#include "Containers/Map.h"
+#include "Logging/LogMacros.h"
+#include "Misc/CString.h"
+#include "Misc/ScopeExit.h"
+#include "Misc/StringBuilder.h"
+#include "String/LexFromString.h"
+
+DEFINE_LOG_CATEGORY_STATIC(LogNativePointerUpgradeTool, Log, All);
+
+IMPLEMENT_APPLICATION(NativePointerUpgradeTool, "NativePointerUpgradeTool");
+
+//////////////////////////////////////////////////////////////////////////
+
+enum class ENativePointerUpgradeBehaviorFlags
+{
+	None,
+	PreviewOnly = 1,
+};
+ENUM_CLASS_FLAGS(ENativePointerUpgradeBehaviorFlags)
+
+class FNativePointerUpgrader
+{
+public:
+
+	FNativePointerUpgrader(const FString& InSCCCommand, ENativePointerUpgradeBehaviorFlags InBehaviorFlags) : SCCCommand(InSCCCommand), BehaviorFlags(InBehaviorFlags)
+	{
+	}
+
+	struct FMemberDeclaration
+	{
+		int32 LineNumber = 0;
+		int32 LineCount = 0;
+		FString TypeName;
+
+		bool operator==(const FMemberDeclaration& Other) const
+		{
+			return (LineNumber == Other.LineNumber) && (TypeName == Other.TypeName);
+		}
+	};
+	using FNativePointerUpgradeList = TMap<FString, TArray<FMemberDeclaration>>;
+
+	TValueOrError<FNativePointerUpgradeList, FString> FillUpgradeList(const FString& LogFilename);
+	bool TryAllUpgrades(const FNativePointerUpgradeList& UpgradeList);
+
+	int32 GetTotalUpgradesFound() const { return TotalUpgradesFound; }
+	int32 GetTotalUpgradesPerformed() const { return TotalUpgradesPerformed; }
+
+private:
+	enum class EByteOrderMarkerType
+	{
+		None,
+		UTF8,
+		UCS2LittleEndian,
+		UCS2BigEndian,
+	};
+
+	struct FPendingWrite
+	{
+		EByteOrderMarkerType BOM;
+		TArray<FString> Contents;
+	};
+	using FPendingWrites = TMap<FString, FPendingWrite>;
+	bool TryUpgradesInFile(const FString& FilenameToUpgrade, TConstArrayView<FMemberDeclaration> MemberDeclarations, FPendingWrites& PendingWrites);
+	bool TryFlushPendingWrites(FPendingWrites& PendingWrites);
+
+	// Using these methods instead of the ones in FFileHelper to ensure that we retain the encoding and BOM of the original file as well as the newlines
+	static bool LoadFileToString(FString& OutResult, EByteOrderMarkerType& OutBOM, const TCHAR* InFilename, uint32 ReadFlags = 0);
+	static bool LoadFileToStringArrayWithNewlines(TArray<FString>& OutResult, EByteOrderMarkerType& OutBOM, const TCHAR* InFilename, uint32 ReadFlags = 0);
+	static bool SaveStringToFile(const FStringView& InString, EByteOrderMarkerType InBOM, const TCHAR* InFilename, uint32 WriteFlags = 0);
+	static bool SaveStringArrayWithNewlinesToFile(const TArray<FString>& InStringArray, EByteOrderMarkerType InBOM, const TCHAR* InFilename, uint32 WriteFlags = 0);
+
+	static int32 CountLines(const FStringView Input);
+
+
+	int32 TotalUpgradesFound = 0;
+	int32 TotalUpgradesPerformed = 0;
+	FString SCCCommand;
+	ENativePointerUpgradeBehaviorFlags BehaviorFlags;
+};
+
+FStringView ParseBracketedStringSegment(const FString& InString, int32& InOutCurrentIndex, const TCHAR* BracketStart, const TCHAR* BracketEnd)
+{
+	int32 ParsedStartIndex = InString.Find(BracketStart, ESearchCase::IgnoreCase, ESearchDir::FromStart, InOutCurrentIndex);
+	if (ParsedStartIndex == INDEX_NONE)
+	{
+		return FStringView();
+	}
+	ParsedStartIndex += FCString::Strlen(BracketStart);
+	int32 ParsedEndIndex = InString.Find(BracketEnd, ESearchCase::IgnoreCase, ESearchDir::FromStart, ParsedStartIndex);
+	if (ParsedEndIndex == INDEX_NONE)
+	{
+		return FStringView();
+	}
+	InOutCurrentIndex = ParsedEndIndex + FCString::Strlen(BracketEnd);
+	return FStringView(*InString + ParsedStartIndex, ParsedEndIndex - ParsedStartIndex);
+}
+
+TValueOrError<FNativePointerUpgrader::FNativePointerUpgradeList, FString> FNativePointerUpgrader::FillUpgradeList(const FString& LogFilename)
+{
+	FNativePointerUpgradeList UpgradeList;
+	const TCHAR Preamble[] = TEXT("Native pointer usage in member declaration detected in '");
+	TArray<FString> NativePointerUpgradeEntries;
+	if (!FFileHelper::LoadFileToStringArrayWithPredicate(NativePointerUpgradeEntries, *LogFilename, [&Preamble](const FString& Line) { return Line.Contains(Preamble);  }))
+	{
+		return MakeError(FString::Printf(TEXT("Unable to load UHT log: %s"), *LogFilename));
+	}
+
+	constexpr int32 PreambleLen = UE_ARRAY_COUNT(Preamble) - 1;
+	for (const FString& Entry : NativePointerUpgradeEntries)
+	{
+		int32 CurrentIndex = 0;
+		FStringView FilenameView = ParseBracketedStringSegment(Entry, CurrentIndex, Preamble, TEXT("'"));
+		if (FilenameView.IsEmpty())
+		{
+			return MakeError(FString::Printf(TEXT("Failed to parse filename from native pointer member upgrade statement: %s"), *Entry));
+		}
+
+		FStringView LineNumberView = ParseBracketedStringSegment(Entry, CurrentIndex, TEXT(", line "), TEXT(" "));
+		if (LineNumberView.IsEmpty())
+		{
+			return MakeError(FString::Printf(TEXT("Failed to parse line number from native pointer member upgrade statement: %s"), *Entry));
+		}
+
+		FStringView TypeNameView = ParseBracketedStringSegment(Entry, CurrentIndex, TEXT("[["), TEXT("]]"));
+		if (TypeNameView.IsEmpty())
+		{
+			return MakeError(FString::Printf(TEXT("Failed to parse type name from native pointer member upgrade statement: %s"), *Entry));
+		}
+
+		int32 LineNumber = -1;
+		LexFromString(LineNumber, LineNumberView);
+		if (LineNumber < 1)
+		{
+			return MakeError(FString::Printf(TEXT("Failed to convert line number to integer from native pointer member upgrade statement: %s"), *Entry));
+		}
+
+		if (!TypeNameView.EndsWith(TEXT("*")))
+		{
+			return MakeError(FString::Printf(TEXT("Failed to find expected asterisk at end of native pointer member upgrade statement: %s"), *Entry));
+		}
+
+		TArray<FMemberDeclaration>& DeclarationArray = UpgradeList.FindOrAdd(FString(FilenameView));
+		FMemberDeclaration NewDecl;
+		NewDecl.LineNumber = LineNumber;
+		NewDecl.TypeName = TypeNameView;
+		NewDecl.TypeName = NewDecl.TypeName.ReplaceEscapedCharWithChar();
+		NewDecl.LineCount = CountLines(NewDecl.TypeName);
+		int32 DeclArraySize = DeclarationArray.Num();
+		DeclarationArray.AddUnique(NewDecl);
+		if (DeclArraySize != DeclarationArray.Num())
+		{
+			TotalUpgradesFound++;
+		}
+	}
+
+	return MakeValue(MoveTemp(UpgradeList));
+}
+
+// The following requirements need to be kept in mind while upgrading source files:
+// 1) If the file section we're updating isn't what we expect it to be, we shouldn't attempt the upgrade
+// 2) We should handle multi-line type declarations
+// 3) We should handle multiple type replacements on the same line
+// 4) We should retain the encoding of the original file
+// 5) We should retain the newlines of the original file
+// 6) We should attempt to retain the indentation (before and after type declaration) of any lines we modify
+bool FNativePointerUpgrader::TryUpgradesInFile(const FString& FilenameToUpgrade, TConstArrayView<FMemberDeclaration> MemberDeclarations, FPendingWrites& PendingWrites)
+{
+	EByteOrderMarkerType OriginalBOM;
+	TArray<FString> FileContentsToUpgrade;
+	if (!LoadFileToStringArrayWithNewlines(FileContentsToUpgrade, OriginalBOM, *FilenameToUpgrade))
+	{
+		UE_LOG(LogNativePointerUpgradeTool, Error, TEXT("Unable to load file to upgrade: %s"), *FilenameToUpgrade);
+		return true;
+	}
+
+	bool bChangesMadeToFileContents = false;
+	for (const FMemberDeclaration& MemberDeclaration : MemberDeclarations)
+	{
+		if (((MemberDeclaration.LineNumber-MemberDeclaration.LineCount) < 0) || (MemberDeclaration.LineNumber > FileContentsToUpgrade.Num()))
+		{
+			UE_LOG(LogNativePointerUpgradeTool, Error, TEXT("Member declaration line '%d' is out of range for the lines (%d) in file '%s' "), MemberDeclaration.LineNumber, FileContentsToUpgrade.Num(), *FilenameToUpgrade);
+			continue;
+		}
+
+		FString LinesToUpgrade;
+		for (int32 LineIndex = MemberDeclaration.LineNumber-MemberDeclaration.LineCount; LineIndex < MemberDeclaration.LineNumber; ++LineIndex)
+		{
+			LinesToUpgrade.Append(FileContentsToUpgrade[LineIndex]);
+		}
+		FString OriginalLinesToUpgrade = LinesToUpgrade;
+
+		// Chop off the trailing asterisk in the member type declaration
+		FString TemplateArg = MemberDeclaration.TypeName.LeftChop(1);
+		// Any trailing whitespace that was after the asterisk should not go in the template argument, but be appended after the template argument
+		// This only applies to space (' ') and tab ('\t') that have non-whitespace before them on the line, and should not include newlines.
+		// This is to avoid having end template bracket ('>') get put on a line that may have a line comment ahead of it ("//").
+		int32 TrailingWhitespaceCount = 0;
+		int32 CharIndex;
+		for (CharIndex = TemplateArg.Len()-1; (CharIndex >= 0) && (TemplateArg[CharIndex] == TCHAR(' ') || TemplateArg[CharIndex] == TCHAR('\t')); --CharIndex)
+		{
+			++TrailingWhitespaceCount;
+		}
+		if ((CharIndex == 0) || (TemplateArg[CharIndex-1] == TEXT('\n')) || (TemplateArg[CharIndex-1] == TEXT('\r')))
+		{
+			// The trailing whitespace represented the entire lead of the line the end template character will be on.  Put the indentation inside the template arg instead of outside of it.
+			TrailingWhitespaceCount = 0;
+		}
+		FString TrailingWhitespace = TemplateArg.Right(TrailingWhitespaceCount);
+		TemplateArg.LeftChopInline(TrailingWhitespaceCount);
+		FString UpgradedTypeName = FString::Printf(TEXT("TObjectPtr<%s>%s"), *TemplateArg, *TrailingWhitespace);
+		if (LinesToUpgrade.ReplaceInline(*MemberDeclaration.TypeName, *UpgradedTypeName) < 1)
+		{
+			if (!LinesToUpgrade.Contains(*UpgradedTypeName))
+			{
+				UE_LOG(LogNativePointerUpgradeTool, Warning, TEXT("%s(%d): Type name string '%s' not found for upgrading, skipping upgrading this line:\n\t%s"), *FilenameToUpgrade, MemberDeclaration.LineNumber, *MemberDeclaration.TypeName, *LinesToUpgrade.TrimStartAndEnd());
+			}
+			continue;
+		}
+
+		TArray<FTextRange> UpgradedTextRanges;
+		FTextRange::CalculateLineRangesFromString(LinesToUpgrade, UpgradedTextRanges);
+
+		for (int32 LineIndex = 0; LineIndex < MemberDeclaration.LineCount; ++LineIndex)
+		{
+			int32 BeginIndex = UpgradedTextRanges[LineIndex].BeginIndex;
+			int32 EndIndex = LineIndex < UpgradedTextRanges.Num()-1 ? UpgradedTextRanges[LineIndex+1].BeginIndex : UpgradedTextRanges[LineIndex].EndIndex;
+			FString UpgradedLine = LinesToUpgrade.Mid(BeginIndex, EndIndex - BeginIndex);
+			FileContentsToUpgrade[MemberDeclaration.LineNumber - MemberDeclaration.LineCount + LineIndex] = UpgradedLine;
+		}
+
+		bChangesMadeToFileContents = true;
+		TotalUpgradesPerformed++;
+		UE_LOG(LogNativePointerUpgradeTool, Display, TEXT("%s(%d): %s\n\tFrom:\t%s\n\tTo:\t%s"), *FilenameToUpgrade, MemberDeclaration.LineNumber, EnumHasAllFlags(BehaviorFlags, ENativePointerUpgradeBehaviorFlags::PreviewOnly) ? TEXT("Previewing upgrade") : TEXT("Upgrading"), *OriginalLinesToUpgrade.TrimStartAndEnd(), *LinesToUpgrade.TrimStartAndEnd());
+	}
+
+	if (bChangesMadeToFileContents)
+	{
+		FPendingWrite& NewWrite = PendingWrites.Add(FilenameToUpgrade);
+		NewWrite.BOM = OriginalBOM;
+		NewWrite.Contents.Append(MoveTemp(FileContentsToUpgrade));
+	}
+
+	return true;
+}
+
+bool FNativePointerUpgrader::TryFlushPendingWrites(FPendingWrites& PendingWrites)
+{
+	if (PendingWrites.IsEmpty())
+	{
+		return true;
+	}
+
+	// Perform a batched SCC command operation
+	if (!SCCCommand.IsEmpty())
+	{
+		TStringBuilder<2048> CompositeFilenamesBuilder;
+		for (const FPendingWrites::ElementType& PendingWrite : PendingWrites)
+		{
+			CompositeFilenamesBuilder.Append(TEXT('"'));
+			CompositeFilenamesBuilder.Append(PendingWrite.Key);
+			CompositeFilenamesBuilder.Append(TEXT("\" "));
+		}
+
+		FString FileSCCCommand(SCCCommand);
+		FileSCCCommand.ReplaceInline(TEXT("{Filenames}"), *CompositeFilenamesBuilder);
+		if (EnumHasAllFlags(BehaviorFlags, ENativePointerUpgradeBehaviorFlags::PreviewOnly))
+		{
+			UE_LOG(LogNativePointerUpgradeTool, Display, TEXT("Would perform SCC command: %s"), *FileSCCCommand);
+		}
+		else
+		{
+			UE_LOG(LogNativePointerUpgradeTool, Display, TEXT("Perform SCC command: %s"), *FileSCCCommand);
+		}
+
+		if (!EnumHasAllFlags(BehaviorFlags, ENativePointerUpgradeBehaviorFlags::PreviewOnly))
+		{
+			FString Params = FCommandLine::RemoveExeName(*FileSCCCommand);
+			FString Cmd = FileSCCCommand.LeftChop(Params.Len()).TrimEnd();
+
+			FProcHandle ProcHandle = FPlatformProcess::CreateProc(*Cmd, *Params, false, false, false, nullptr, 0, nullptr, nullptr);
+			if (!ProcHandle.IsValid())
+			{
+				UE_LOG(LogNativePointerUpgradeTool, Error, TEXT("Failed to perform SCC command: %s"), *FileSCCCommand);
+				return false;
+			}
+
+			FPlatformProcess::WaitForProc(ProcHandle);
+			int32 RetCode = -1;
+			FPlatformProcess::GetProcReturnCode(ProcHandle, &RetCode);
+			FPlatformProcess::CloseProc(ProcHandle);
+
+			if (RetCode != 0)
+			{
+				UE_LOG(LogNativePointerUpgradeTool, Error, TEXT("SCC command returned an unexpected return code: %d"), RetCode);
+				return false;
+			}
+		}
+	}
+
+	for (const FPendingWrites::ElementType& PendingWrite : PendingWrites)
+	{
+		if (EnumHasAllFlags(BehaviorFlags, ENativePointerUpgradeBehaviorFlags::PreviewOnly))
+		{
+			UE_LOG(LogNativePointerUpgradeTool, Display, TEXT("File would be modified and saved: %s"), *PendingWrite.Key);
+		}
+		else
+		{
+			UE_LOG(LogNativePointerUpgradeTool, Display, TEXT("Saving modified file: %s"), *PendingWrite.Key);
+			if (!SaveStringArrayWithNewlinesToFile(PendingWrite.Value.Contents, PendingWrite.Value.BOM, *PendingWrite.Key))
+			{
+				UE_LOG(LogNativePointerUpgradeTool, Error, TEXT("Failed to save upgraded file: %s"), *PendingWrite.Key);
+				return false;
+			}
+		}
+	}
+
+	PendingWrites.Empty();
+	return true;
+}
+
+bool FNativePointerUpgrader::TryAllUpgrades(const FNativePointerUpgradeList& UpgradeList)
+{
+	FPendingWrites PendingWrites;
+	for (const TPair<FString, TArray<FMemberDeclaration>>& FileToUpgrade : UpgradeList)
+	{
+		if (!TryUpgradesInFile(FileToUpgrade.Key, FileToUpgrade.Value, PendingWrites))
+		{
+			return false;
+		}
+
+		if (PendingWrites.Num() >= 20)
+		{
+			if (!TryFlushPendingWrites(PendingWrites))
+			{
+				return false;
+			}
+			check(PendingWrites.IsEmpty());
+		}
+	}
+
+	if (!TryFlushPendingWrites(PendingWrites))
+	{
+		return false;
+	}
+	return true;
+}
+
+bool FNativePointerUpgrader::LoadFileToString(FString& OutResult, EByteOrderMarkerType& OutBOM, const TCHAR* InFilename, uint32 ReadFlags)
+{
+	TUniquePtr<FArchive> Reader(IFileManager::Get().CreateFileReader(InFilename, ReadFlags));
+	if (!Reader)
+	{
+		return false;
+	}
+
+	FScopedLoadingState ScopedLoadingState(*Reader->GetArchiveName());
+
+	int64 Size = Reader->TotalSize();
+	if (!Size)
+	{
+		OutBOM = EByteOrderMarkerType::None;
+		OutResult.Empty();
+		return true;
+	}
+
+	uint8* OriginalBuffer = (uint8*)FMemory::Malloc(Size);
+	uint8* Buffer = OriginalBuffer;
+	ON_SCOPE_EXIT
+	{
+		FMemory::Free(OriginalBuffer);
+	};
+
+	Reader->Serialize(Buffer, Size);
+	if (!Reader->Close())
+	{
+		return false;
+	}
+
+	TArray<TCHAR>& ResultArray = OutResult.GetCharArray();
+	ResultArray.Empty();
+
+	bool bIsUnicode = false;
+	if (Size >= 2 && !(Size & 1) && Buffer[0] == 0xff && Buffer[1] == 0xfe)
+	{
+		// Unicode Intel byte order. Less 1 for the FFFE header, additional 1 for null terminator.
+		OutBOM = EByteOrderMarkerType::UCS2LittleEndian;
+		ResultArray.AddUninitialized(Size / 2);
+		for (int32 i = 0; i < (Size / 2) - 1; i++)
+		{
+			ResultArray[i] = CharCast<TCHAR>((UCS2CHAR)((uint16)Buffer[i * 2 + 2] + (uint16)Buffer[i * 2 + 3] * 256));
+		}
+		bIsUnicode = true;
+	}
+	else if (Size >= 2 && !(Size & 1) && Buffer[0] == 0xfe && Buffer[1] == 0xff)
+	{
+		// Unicode non-Intel byte order. Less 1 for the FFFE header, additional 1 for null terminator.
+		OutBOM = EByteOrderMarkerType::UCS2BigEndian;
+		ResultArray.AddUninitialized(Size / 2);
+		for (int32 i = 0; i < (Size / 2) - 1; i++)
+		{
+			ResultArray[i] = CharCast<TCHAR>((UCS2CHAR)((uint16)Buffer[i * 2 + 3] + (uint16)Buffer[i * 2 + 2] * 256));
+		}
+		bIsUnicode = true;
+	}
+	else
+	{
+		if (Size >= 3 && Buffer[0] == 0xef && Buffer[1] == 0xbb && Buffer[2] == 0xbf)
+		{
+			// Skip over UTF-8 BOM if there is one
+			OutBOM = EByteOrderMarkerType::UTF8;
+			Buffer += 3;
+			Size -= 3;
+		}
+		else
+		{
+			OutBOM = EByteOrderMarkerType::None;
+		}
+
+		int32 Length = FUTF8ToTCHAR_Convert::ConvertedLength(reinterpret_cast<const ANSICHAR*>(Buffer), Size);
+		ResultArray.AddUninitialized(Length + 1); // +1 for the null terminator
+		FUTF8ToTCHAR_Convert::Convert(ResultArray.GetData(), ResultArray.Num(), reinterpret_cast<const ANSICHAR*>(Buffer), Size);
+		ResultArray[Length] = TEXT('\0');
+	}
+
+	if (ResultArray.Num() == 1)
+	{
+		// If it's only a zero terminator then make the result actually empty
+		ResultArray.Empty();
+	}
+	else
+	{
+		// Else ensure null terminator is present
+		ResultArray.Last() = 0;
+
+		if (bIsUnicode)
+		{
+			// Inline combine any surrogate pairs in the data when loading into a UTF-32 string
+			StringConv::InlineCombineSurrogates(OutResult);
+		}
+	}
+
+	return true;
+}
+
+bool FNativePointerUpgrader::LoadFileToStringArrayWithNewlines(TArray<FString>& OutResult, EByteOrderMarkerType& OutBOM, const TCHAR* InFilename, uint32 ReadFlags)
+{
+	FString FileContents;
+	if (!LoadFileToString(FileContents, OutBOM, InFilename, ReadFlags))
+	{
+		return false;
+	}
+
+	TArray<FTextRange> LineRanges;
+	FTextRange::CalculateLineRangesFromString(FileContents, LineRanges);
+	OutResult.Reset(LineRanges.Num());
+
+	for (int32 LineIndex = 0; LineIndex < LineRanges.Num(); ++LineIndex)
+	{
+		int32 BeginIndex = LineRanges[LineIndex].BeginIndex;
+		int32 EndIndex = (LineIndex < (LineRanges.Num() - 1)) ? LineRanges[LineIndex+1].BeginIndex : LineRanges[LineIndex].EndIndex;
+		OutResult.Add(FileContents.Mid(BeginIndex, EndIndex - BeginIndex));
+	}
+
+	return true;
+}
+
+bool FNativePointerUpgrader::SaveStringToFile(const FStringView& InString, EByteOrderMarkerType InBOM, const TCHAR* InFilename, uint32 WriteFlags)
+{
+	// max size of the string is a UCS2CHAR for each character and some UNICODE magic 
+	TUniquePtr<FArchive> Ar = TUniquePtr<FArchive>(IFileManager::Get().CreateFileWriter(InFilename, WriteFlags));
+	if (!Ar)
+	{
+		return false;
+	}
+
+	if (InString.IsEmpty())
+	{
+		return true;
+	}
+
+	if (InBOM == EByteOrderMarkerType::UTF8)
+	{
+		UTF8CHAR UTF8BOM[] = {0xEF, 0xBB, 0xBF};
+		Ar->Serialize(&UTF8BOM, UE_ARRAY_COUNT(UTF8BOM) * sizeof(UTF8CHAR));
+
+		FTCHARToUTF8 UTF8String(InString.GetData(), InString.Len());
+		Ar->Serialize((UTF8CHAR*)UTF8String.Get(), UTF8String.Length() * sizeof(UTF8CHAR));
+	}
+	else if ((InBOM == EByteOrderMarkerType::UCS2LittleEndian) || (InBOM == EByteOrderMarkerType::UCS2BigEndian))
+	{
+		Ar->SetByteSwapping(PLATFORM_LITTLE_ENDIAN ? InBOM == EByteOrderMarkerType::UCS2BigEndian : InBOM == EByteOrderMarkerType::UCS2LittleEndian);
+		
+		UTF16CHAR BOM = UNICODE_BOM;
+		Ar->Serialize(&BOM, sizeof(UTF16CHAR));
+
+		// Note: This is a no-op on platforms that are using a 16-bit TCHAR
+		FTCHARToUTF16 UTF16String(InString.GetData(), InString.Len());
+		Ar->Serialize((UTF16CHAR*)UTF16String.Get(), UTF16String.Length() * sizeof(UTF16CHAR));
+	}
+	else
+	{
+		// If there is no BOM specified, write to UTF8 without a BOM.  If the contents fit within ASCII, they will be
+		// written as ASCII.  If they don't fit, they'll be written as multibyte code points.  This is necessary as
+		// a file was encountered that is UTF8 without any BOM.  The reading code handles this correctly,
+		// so we want to ensure that the writing code doesn't substitute "bogus chars" (?) for any non-ASCII characters
+		// at the time of writing.
+		FTCHARToUTF8 UTF8String(InString.GetData(), InString.Len());
+		Ar->Serialize((UTF8CHAR*)UTF8String.Get(), UTF8String.Length() * sizeof(UTF8CHAR));
+	}
+
+	// Always explicitly close to catch errors from flush/close
+	Ar->Close();
+
+	return !Ar->IsError() && !Ar->IsCriticalError();
+}
+
+bool FNativePointerUpgrader::SaveStringArrayWithNewlinesToFile(const TArray<FString>& InStringArray, EByteOrderMarkerType InBOM, const TCHAR* InFilename, uint32 WriteFlags)
+{
+	FString ComposedFileContents = FString::Join(InStringArray, TEXT(""));
+	return SaveStringToFile(ComposedFileContents, InBOM, InFilename, WriteFlags);
+}
+
+int32 FNativePointerUpgrader::CountLines(const FStringView Input)
+{
+	int32 LineCount = 0;
+	int32 LineBeginIndex = 0;
+
+	// Loop through splitting at new-lines
+	const TCHAR* const InputStart = Input.GetData();
+	const TCHAR* const InputEnd = InputStart + Input.Len();
+	for (const TCHAR* CurrentChar = InputStart; CurrentChar != InputEnd; ++CurrentChar)
+	{
+		// Handle a chain of \r\n slightly differently to stop the FChar::IsLinebreak adding two separate new-lines
+		const bool bIsWindowsNewLine = ((CurrentChar + 1) != InputEnd) && (*CurrentChar == '\r' && *(CurrentChar + 1) == '\n');
+		if (bIsWindowsNewLine || FChar::IsLinebreak(*CurrentChar))
+		{
+			++LineCount;
+
+			if (bIsWindowsNewLine)
+			{
+				++CurrentChar; // skip the \n of the \r\n chain
+			}
+			LineBeginIndex = UE_PTRDIFF_TO_INT32(CurrentChar - InputStart) + 1; // The next line begins after the end of the current line
+		}
+	}
+
+	// Process any remaining string after the last new-line
+	if (LineBeginIndex <= Input.Len())
+	{
+		++LineCount;
+	}
+
+	return LineCount;
+}
+
+//////////////////////////////////////////////////////////////////////////
+
+INT32_MAIN_INT32_ARGC_TCHAR_ARGV()
+{
+	FString CmdLine = FCommandLine::BuildFromArgV(nullptr, ArgC, ArgV, nullptr);
+	GEngineLoop.PreInit(*CmdLine);
+
+	// Make sure the engine is properly cleaned up whenever we exit this function
+	ON_SCOPE_EXIT
+	{
+		FEngineLoop::AppPreExit();
+		FEngineLoop::AppExit();
+	};
+
+	FString UHTLogFilename;
+	if (CmdLine.Len() && **CmdLine != TEXT('-'))
+	{
+		const TCHAR* CmdLinePtr = *CmdLine;
+		UHTLogFilename = FParse::Token(CmdLinePtr, false);
+	}
+
+	if (UHTLogFilename.IsEmpty())
+	{
+		UHTLogFilename = FPaths::ConvertRelativePathToFull(FPlatformProcess::BaseDir(), TEXT("../../../Engine/Programs/UnrealHeaderTool/Saved/Logs/UnrealHeaderTool.log"));
+	}
+	
+	ENativePointerUpgradeBehaviorFlags BehaviorFlags = ENativePointerUpgradeBehaviorFlags::None;
+	FString SCCCommand;
+	FParse::Value(*CmdLine, TEXT("-SCCCommand="), SCCCommand);
+	if (FParse::Param(*CmdLine, TEXT("n")) || FParse::Param(*CmdLine, TEXT("PREVIEW")))
+	{
+		BehaviorFlags |= ENativePointerUpgradeBehaviorFlags::PreviewOnly;
+	}
+
+	FNativePointerUpgrader NativePointerUpgrader(SCCCommand, BehaviorFlags);
+
+	TValueOrError<FNativePointerUpgrader::FNativePointerUpgradeList, FString> UpgradeListResult =
+		NativePointerUpgrader.FillUpgradeList(UHTLogFilename);
+
+	if (!UpgradeListResult.IsValid())
+	{
+		UE_LOG(LogNativePointerUpgradeTool, Error, TEXT("Upgrade log parsing error: %s"), *UpgradeListResult.GetError());
+		return 1;
+	}
+
+	if (!NativePointerUpgrader.TryAllUpgrades(UpgradeListResult.GetValue()))
+	{
+		return 1;
+	}
+
+	if (EnumHasAllFlags(BehaviorFlags, ENativePointerUpgradeBehaviorFlags::PreviewOnly))
+	{
+		UE_LOG(LogNativePointerUpgradeTool, Display, TEXT("Previewed upgrade successfully.  %d upgrades found; %d upgrades would be performed"), NativePointerUpgrader.GetTotalUpgradesFound(), NativePointerUpgrader.GetTotalUpgradesPerformed());
+	}
+	else
+	{
+		UE_LOG(LogNativePointerUpgradeTool, Display, TEXT("Upgraded successfully.  %d upgrades found; %d upgrades performed"), NativePointerUpgrader.GetTotalUpgradesFound(), NativePointerUpgrader.GetTotalUpgradesPerformed());
+	}
+
+	return 0;
+}
diff --git a/Engine/Source/Programs/UnrealObjectPtrTool/UnrealObjectPtrTool.Build.cs b/Engine/Source/Programs/UnrealObjectPtrTool/UnrealObjectPtrTool.Build.cs
new file mode 100644
--- /dev/null
+++ b/Engine/Source/Programs/UnrealObjectPtrTool/UnrealObjectPtrTool.Build.cs
@@ -0,0 +1,16 @@
+// Copyright Epic Games, Inc. All Rights Reserved.
+
+using UnrealBuildTool;
+
+public class NativePointerUpgradeTool : ModuleRules
+{
+	public NativePointerUpgradeTool(ReadOnlyTargetRules Target) : base(Target)
+	{
+		PublicIncludePaths.Add("Runtime/Launch/Public");
+
+		PrivateIncludePaths.Add("Runtime/Launch/Private");		// For LaunchEngineLoop.cpp include
+
+		PrivateDependencyModuleNames.Add("Core");
+		PrivateDependencyModuleNames.Add("Projects");
+	}
+}
diff --git a/Engine/Source/Programs/UnrealObjectPtrTool/UnrealObjectPtrTool.Target.cs b/Engine/Source/Programs/UnrealObjectPtrTool/UnrealObjectPtrTool.Target.cs
new file mode 100644
--- /dev/null
+++ b/Engine/Source/Programs/UnrealObjectPtrTool/UnrealObjectPtrTool.Target.cs
@@ -0,0 +1,24 @@
+// Copyright Epic Games, Inc. All Rights Reserved.
+
+using UnrealBuildTool;
+using System.Collections.Generic;
+
+[SupportedPlatforms(UnrealPlatformClass.Desktop)]
+public class NativePointerUpgradeToolTarget : TargetRules
+{
+	public NativePointerUpgradeToolTarget(TargetInfo Target) : base(Target)
+	{
+		Type = TargetType.Program;
+		LinkType = TargetLinkType.Modular;
+		LaunchModuleName = "NativePointerUpgradeTool";
+
+		bBuildDeveloperTools			= false;
+		bUseMallocProfiler				= false;
+		bBuildWithEditorOnlyData		= false;
+		bCompileAgainstEngine			= false;
+		bCompileAgainstCoreUObject		= false;
+		bCompileAgainstApplicationCore	= false;
+		bUsesSlate = false;
+		bIsBuildingConsoleApplication	= true;
+	}
+}
