File content changes to support renaming NativePointerUpgradeTool to UnrealObjectPtrTool.

diff --git a/Engine/Source/Programs/UnrealObjectPtrTool/Private/UnrealObjectPtrTool.cpp b/Engine/Source/Programs/UnrealObjectPtrTool/Private/UnrealObjectPtrTool.cpp
--- a/Engine/Source/Programs/UnrealObjectPtrTool/Private/UnrealObjectPtrTool.cpp
+++ b/Engine/Source/Programs/UnrealObjectPtrTool/Private/UnrealObjectPtrTool.cpp
@@ -12,24 +12,24 @@
 #include "Misc/StringBuilder.h"
 #include "String/LexFromString.h"
 
-DEFINE_LOG_CATEGORY_STATIC(LogNativePointerUpgradeTool, Log, All);
+DEFINE_LOG_CATEGORY_STATIC(LogUnrealObjectPtrTool, Log, All);
 
-IMPLEMENT_APPLICATION(NativePointerUpgradeTool, "NativePointerUpgradeTool");
+IMPLEMENT_APPLICATION(UnrealObjectPtrTool, "UnrealObjectPtrTool");
 
 //////////////////////////////////////////////////////////////////////////
 
-enum class ENativePointerUpgradeBehaviorFlags
+enum class EPointerUpgradeBehaviorFlags
 {
 	None,
 	PreviewOnly = 1,
 };
-ENUM_CLASS_FLAGS(ENativePointerUpgradeBehaviorFlags)
+ENUM_CLASS_FLAGS(EPointerUpgradeBehaviorFlags)
 
-class FNativePointerUpgrader
+class FPointerUpgrader
 {
 public:
 
-	FNativePointerUpgrader(const FString& InSCCCommand, ENativePointerUpgradeBehaviorFlags InBehaviorFlags) : SCCCommand(InSCCCommand), BehaviorFlags(InBehaviorFlags)
+	FPointerUpgrader(const FString& InSCCCommand, EPointerUpgradeBehaviorFlags InBehaviorFlags) : SCCCommand(InSCCCommand), BehaviorFlags(InBehaviorFlags)
 	{
 	}
 
@@ -44,10 +44,10 @@ public:
 			return (LineNumber == Other.LineNumber) && (TypeName == Other.TypeName);
 		}
 	};
-	using FNativePointerUpgradeList = TMap<FString, TArray<FMemberDeclaration>>;
+	using FPointerUpgradeList = TMap<FString, TArray<FMemberDeclaration>>;
 
-	TValueOrError<FNativePointerUpgradeList, FString> FillUpgradeList(const FString& LogFilename);
-	bool TryAllUpgrades(const FNativePointerUpgradeList& UpgradeList);
+	TValueOrError<FPointerUpgradeList, FString> FillUpgradeList(const FString& LogFilename);
+	bool TryAllUpgrades(const FPointerUpgradeList& UpgradeList);
 
 	int32 GetTotalUpgradesFound() const { return TotalUpgradesFound; }
 	int32 GetTotalUpgradesPerformed() const { return TotalUpgradesPerformed; }
@@ -82,7 +82,7 @@ private:
 	int32 TotalUpgradesFound = 0;
 	int32 TotalUpgradesPerformed = 0;
 	FString SCCCommand;
-	ENativePointerUpgradeBehaviorFlags BehaviorFlags;
+	EPointerUpgradeBehaviorFlags BehaviorFlags;
 };
 
 FStringView ParseBracketedStringSegment(const FString& InString, int32& InOutCurrentIndex, const TCHAR* BracketStart, const TCHAR* BracketEnd)
@@ -102,18 +102,18 @@ FStringView ParseBracketedStringSegment(const FString& InString, int32& InOutCur
 	return FStringView(*InString + ParsedStartIndex, ParsedEndIndex - ParsedStartIndex);
 }
 
-TValueOrError<FNativePointerUpgrader::FNativePointerUpgradeList, FString> FNativePointerUpgrader::FillUpgradeList(const FString& LogFilename)
+TValueOrError<FPointerUpgrader::FPointerUpgradeList, FString> FPointerUpgrader::FillUpgradeList(const FString& LogFilename)
 {
-	FNativePointerUpgradeList UpgradeList;
+	FPointerUpgradeList UpgradeList;
 	const TCHAR Preamble[] = TEXT("Native pointer usage in member declaration detected in '");
-	TArray<FString> NativePointerUpgradeEntries;
-	if (!FFileHelper::LoadFileToStringArrayWithPredicate(NativePointerUpgradeEntries, *LogFilename, [&Preamble](const FString& Line) { return Line.Contains(Preamble);  }))
+	TArray<FString> PointerUpgradeEntries;
+	if (!FFileHelper::LoadFileToStringArrayWithPredicate(PointerUpgradeEntries, *LogFilename, [&Preamble](const FString& Line) { return Line.Contains(Preamble);  }))
 	{
 		return MakeError(FString::Printf(TEXT("Unable to load UHT log: %s"), *LogFilename));
 	}
 
 	constexpr int32 PreambleLen = UE_ARRAY_COUNT(Preamble) - 1;
-	for (const FString& Entry : NativePointerUpgradeEntries)
+	for (const FString& Entry : PointerUpgradeEntries)
 	{
 		int32 CurrentIndex = 0;
 		FStringView FilenameView = ParseBracketedStringSegment(Entry, CurrentIndex, Preamble, TEXT("'"));
@@ -170,13 +170,13 @@ TValueOrError<FNativePointerUpgrader::FNativePointerUpgradeList, FString> FNativ
 // 4) We should retain the encoding of the original file
 // 5) We should retain the newlines of the original file
 // 6) We should attempt to retain the indentation (before and after type declaration) of any lines we modify
-bool FNativePointerUpgrader::TryUpgradesInFile(const FString& FilenameToUpgrade, TConstArrayView<FMemberDeclaration> MemberDeclarations, FPendingWrites& PendingWrites)
+bool FPointerUpgrader::TryUpgradesInFile(const FString& FilenameToUpgrade, TConstArrayView<FMemberDeclaration> MemberDeclarations, FPendingWrites& PendingWrites)
 {
 	EByteOrderMarkerType OriginalBOM;
 	TArray<FString> FileContentsToUpgrade;
 	if (!LoadFileToStringArrayWithNewlines(FileContentsToUpgrade, OriginalBOM, *FilenameToUpgrade))
 	{
-		UE_LOG(LogNativePointerUpgradeTool, Error, TEXT("Unable to load file to upgrade: %s"), *FilenameToUpgrade);
+		UE_LOG(LogUnrealObjectPtrTool, Error, TEXT("Unable to load file to upgrade: %s"), *FilenameToUpgrade);
 		return true;
 	}
 
@@ -185,7 +185,7 @@ bool FNativePointerUpgrader::TryUpgradesInFile(const FString& FilenameToUpgrade,
 	{
 		if (((MemberDeclaration.LineNumber-MemberDeclaration.LineCount) < 0) || (MemberDeclaration.LineNumber > FileContentsToUpgrade.Num()))
 		{
-			UE_LOG(LogNativePointerUpgradeTool, Error, TEXT("Member declaration line '%d' is out of range for the lines (%d) in file '%s' "), MemberDeclaration.LineNumber, FileContentsToUpgrade.Num(), *FilenameToUpgrade);
+			UE_LOG(LogUnrealObjectPtrTool, Error, TEXT("Member declaration line '%d' is out of range for the lines (%d) in file '%s' "), MemberDeclaration.LineNumber, FileContentsToUpgrade.Num(), *FilenameToUpgrade);
 			continue;
 		}
 
@@ -219,7 +219,7 @@ bool FNativePointerUpgrader::TryUpgradesInFile(const FString& FilenameToUpgrade,
 		{
 			if (!LinesToUpgrade.Contains(*UpgradedTypeName))
 			{
-				UE_LOG(LogNativePointerUpgradeTool, Warning, TEXT("%s(%d): Type name string '%s' not found for upgrading, skipping upgrading this line:\n\t%s"), *FilenameToUpgrade, MemberDeclaration.LineNumber, *MemberDeclaration.TypeName, *LinesToUpgrade.TrimStartAndEnd());
+				UE_LOG(LogUnrealObjectPtrTool, Warning, TEXT("%s(%d): Type name string '%s' not found for upgrading, skipping upgrading this line:\n\t%s"), *FilenameToUpgrade, MemberDeclaration.LineNumber, *MemberDeclaration.TypeName, *LinesToUpgrade.TrimStartAndEnd());
 			}
 			continue;
 		}
@@ -237,7 +237,7 @@ bool FNativePointerUpgrader::TryUpgradesInFile(const FString& FilenameToUpgrade,
 
 		bChangesMadeToFileContents = true;
 		TotalUpgradesPerformed++;
-		UE_LOG(LogNativePointerUpgradeTool, Display, TEXT("%s(%d): %s\n\tFrom:\t%s\n\tTo:\t%s"), *FilenameToUpgrade, MemberDeclaration.LineNumber, EnumHasAllFlags(BehaviorFlags, ENativePointerUpgradeBehaviorFlags::PreviewOnly) ? TEXT("Previewing upgrade") : TEXT("Upgrading"), *OriginalLinesToUpgrade.TrimStartAndEnd(), *LinesToUpgrade.TrimStartAndEnd());
+		UE_LOG(LogUnrealObjectPtrTool, Display, TEXT("%s(%d): %s\n\tFrom:\t%s\n\tTo:\t%s"), *FilenameToUpgrade, MemberDeclaration.LineNumber, EnumHasAllFlags(BehaviorFlags, EPointerUpgradeBehaviorFlags::PreviewOnly) ? TEXT("Previewing upgrade") : TEXT("Upgrading"), *OriginalLinesToUpgrade.TrimStartAndEnd(), *LinesToUpgrade.TrimStartAndEnd());
 	}
 
 	if (bChangesMadeToFileContents)
@@ -250,7 +250,7 @@ bool FNativePointerUpgrader::TryUpgradesInFile(const FString& FilenameToUpgrade,
 	return true;
 }
 
-bool FNativePointerUpgrader::TryFlushPendingWrites(FPendingWrites& PendingWrites)
+bool FPointerUpgrader::TryFlushPendingWrites(FPendingWrites& PendingWrites)
 {
 	if (PendingWrites.IsEmpty())
 	{
@@ -270,16 +270,16 @@ bool FNativePointerUpgrader::TryFlushPendingWrites(FPendingWrites& PendingWrites
 
 		FString FileSCCCommand(SCCCommand);
 		FileSCCCommand.ReplaceInline(TEXT("{Filenames}"), *CompositeFilenamesBuilder);
-		if (EnumHasAllFlags(BehaviorFlags, ENativePointerUpgradeBehaviorFlags::PreviewOnly))
+		if (EnumHasAllFlags(BehaviorFlags, EPointerUpgradeBehaviorFlags::PreviewOnly))
 		{
-			UE_LOG(LogNativePointerUpgradeTool, Display, TEXT("Would perform SCC command: %s"), *FileSCCCommand);
+			UE_LOG(LogUnrealObjectPtrTool, Display, TEXT("Would perform SCC command: %s"), *FileSCCCommand);
 		}
 		else
 		{
-			UE_LOG(LogNativePointerUpgradeTool, Display, TEXT("Perform SCC command: %s"), *FileSCCCommand);
+			UE_LOG(LogUnrealObjectPtrTool, Display, TEXT("Perform SCC command: %s"), *FileSCCCommand);
 		}
 
-		if (!EnumHasAllFlags(BehaviorFlags, ENativePointerUpgradeBehaviorFlags::PreviewOnly))
+		if (!EnumHasAllFlags(BehaviorFlags, EPointerUpgradeBehaviorFlags::PreviewOnly))
 		{
 			FString Params = FCommandLine::RemoveExeName(*FileSCCCommand);
 			FString Cmd = FileSCCCommand.LeftChop(Params.Len()).TrimEnd();
@@ -287,7 +287,7 @@ bool FNativePointerUpgrader::TryFlushPendingWrites(FPendingWrites& PendingWrites
 			FProcHandle ProcHandle = FPlatformProcess::CreateProc(*Cmd, *Params, false, false, false, nullptr, 0, nullptr, nullptr);
 			if (!ProcHandle.IsValid())
 			{
-				UE_LOG(LogNativePointerUpgradeTool, Error, TEXT("Failed to perform SCC command: %s"), *FileSCCCommand);
+				UE_LOG(LogUnrealObjectPtrTool, Error, TEXT("Failed to perform SCC command: %s"), *FileSCCCommand);
 				return false;
 			}
 
@@ -298,7 +298,7 @@ bool FNativePointerUpgrader::TryFlushPendingWrites(FPendingWrites& PendingWrites
 
 			if (RetCode != 0)
 			{
-				UE_LOG(LogNativePointerUpgradeTool, Error, TEXT("SCC command returned an unexpected return code: %d"), RetCode);
+				UE_LOG(LogUnrealObjectPtrTool, Error, TEXT("SCC command returned an unexpected return code: %d"), RetCode);
 				return false;
 			}
 		}
@@ -306,16 +306,16 @@ bool FNativePointerUpgrader::TryFlushPendingWrites(FPendingWrites& PendingWrites
 
 	for (const FPendingWrites::ElementType& PendingWrite : PendingWrites)
 	{
-		if (EnumHasAllFlags(BehaviorFlags, ENativePointerUpgradeBehaviorFlags::PreviewOnly))
+		if (EnumHasAllFlags(BehaviorFlags, EPointerUpgradeBehaviorFlags::PreviewOnly))
 		{
-			UE_LOG(LogNativePointerUpgradeTool, Display, TEXT("File would be modified and saved: %s"), *PendingWrite.Key);
+			UE_LOG(LogUnrealObjectPtrTool, Display, TEXT("File would be modified and saved: %s"), *PendingWrite.Key);
 		}
 		else
 		{
-			UE_LOG(LogNativePointerUpgradeTool, Display, TEXT("Saving modified file: %s"), *PendingWrite.Key);
+			UE_LOG(LogUnrealObjectPtrTool, Display, TEXT("Saving modified file: %s"), *PendingWrite.Key);
 			if (!SaveStringArrayWithNewlinesToFile(PendingWrite.Value.Contents, PendingWrite.Value.BOM, *PendingWrite.Key))
 			{
-				UE_LOG(LogNativePointerUpgradeTool, Error, TEXT("Failed to save upgraded file: %s"), *PendingWrite.Key);
+				UE_LOG(LogUnrealObjectPtrTool, Error, TEXT("Failed to save upgraded file: %s"), *PendingWrite.Key);
 				return false;
 			}
 		}
@@ -325,7 +325,7 @@ bool FNativePointerUpgrader::TryFlushPendingWrites(FPendingWrites& PendingWrites
 	return true;
 }
 
-bool FNativePointerUpgrader::TryAllUpgrades(const FNativePointerUpgradeList& UpgradeList)
+bool FPointerUpgrader::TryAllUpgrades(const FPointerUpgradeList& UpgradeList)
 {
 	FPendingWrites PendingWrites;
 	for (const TPair<FString, TArray<FMemberDeclaration>>& FileToUpgrade : UpgradeList)
@@ -352,7 +352,7 @@ bool FNativePointerUpgrader::TryAllUpgrades(const FNativePointerUpgradeList& Upg
 	return true;
 }
 
-bool FNativePointerUpgrader::LoadFileToString(FString& OutResult, EByteOrderMarkerType& OutBOM, const TCHAR* InFilename, uint32 ReadFlags)
+bool FPointerUpgrader::LoadFileToString(FString& OutResult, EByteOrderMarkerType& OutBOM, const TCHAR* InFilename, uint32 ReadFlags)
 {
 	TUniquePtr<FArchive> Reader(IFileManager::Get().CreateFileReader(InFilename, ReadFlags));
 	if (!Reader)
@@ -449,7 +449,7 @@ bool FNativePointerUpgrader::LoadFileToString(FString& OutResult, EByteOrderMark
 	return true;
 }
 
-bool FNativePointerUpgrader::LoadFileToStringArrayWithNewlines(TArray<FString>& OutResult, EByteOrderMarkerType& OutBOM, const TCHAR* InFilename, uint32 ReadFlags)
+bool FPointerUpgrader::LoadFileToStringArrayWithNewlines(TArray<FString>& OutResult, EByteOrderMarkerType& OutBOM, const TCHAR* InFilename, uint32 ReadFlags)
 {
 	FString FileContents;
 	if (!LoadFileToString(FileContents, OutBOM, InFilename, ReadFlags))
@@ -471,7 +471,7 @@ bool FNativePointerUpgrader::LoadFileToStringArrayWithNewlines(TArray<FString>&
 	return true;
 }
 
-bool FNativePointerUpgrader::SaveStringToFile(const FStringView& InString, EByteOrderMarkerType InBOM, const TCHAR* InFilename, uint32 WriteFlags)
+bool FPointerUpgrader::SaveStringToFile(const FStringView& InString, EByteOrderMarkerType InBOM, const TCHAR* InFilename, uint32 WriteFlags)
 {
 	// max size of the string is a UCS2CHAR for each character and some UNICODE magic 
 	TUniquePtr<FArchive> Ar = TUniquePtr<FArchive>(IFileManager::Get().CreateFileWriter(InFilename, WriteFlags));
@@ -521,13 +521,13 @@ bool FNativePointerUpgrader::SaveStringToFile(const FStringView& InString, EByte
 	return !Ar->IsError() && !Ar->IsCriticalError();
 }
 
-bool FNativePointerUpgrader::SaveStringArrayWithNewlinesToFile(const TArray<FString>& InStringArray, EByteOrderMarkerType InBOM, const TCHAR* InFilename, uint32 WriteFlags)
+bool FPointerUpgrader::SaveStringArrayWithNewlinesToFile(const TArray<FString>& InStringArray, EByteOrderMarkerType InBOM, const TCHAR* InFilename, uint32 WriteFlags)
 {
 	FString ComposedFileContents = FString::Join(InStringArray, TEXT(""));
 	return SaveStringToFile(ComposedFileContents, InBOM, InFilename, WriteFlags);
 }
 
-int32 FNativePointerUpgrader::CountLines(const FStringView Input)
+int32 FPointerUpgrader::CountLines(const FStringView Input)
 {
 	int32 LineCount = 0;
 	int32 LineBeginIndex = 0;
@@ -586,37 +586,37 @@ INT32_MAIN_INT32_ARGC_TCHAR_ARGV()
 		UHTLogFilename = FPaths::ConvertRelativePathToFull(FPlatformProcess::BaseDir(), TEXT("../../../Engine/Programs/UnrealHeaderTool/Saved/Logs/UnrealHeaderTool.log"));
 	}
 	
-	ENativePointerUpgradeBehaviorFlags BehaviorFlags = ENativePointerUpgradeBehaviorFlags::None;
+	EPointerUpgradeBehaviorFlags BehaviorFlags = EPointerUpgradeBehaviorFlags::None;
 	FString SCCCommand;
 	FParse::Value(*CmdLine, TEXT("-SCCCommand="), SCCCommand);
 	if (FParse::Param(*CmdLine, TEXT("n")) || FParse::Param(*CmdLine, TEXT("PREVIEW")))
 	{
-		BehaviorFlags |= ENativePointerUpgradeBehaviorFlags::PreviewOnly;
+		BehaviorFlags |= EPointerUpgradeBehaviorFlags::PreviewOnly;
 	}
 
-	FNativePointerUpgrader NativePointerUpgrader(SCCCommand, BehaviorFlags);
+	FPointerUpgrader PointerUpgrader(SCCCommand, BehaviorFlags);
 
-	TValueOrError<FNativePointerUpgrader::FNativePointerUpgradeList, FString> UpgradeListResult =
-		NativePointerUpgrader.FillUpgradeList(UHTLogFilename);
+	TValueOrError<FPointerUpgrader::FPointerUpgradeList, FString> UpgradeListResult =
+		PointerUpgrader.FillUpgradeList(UHTLogFilename);
 
 	if (!UpgradeListResult.IsValid())
 	{
-		UE_LOG(LogNativePointerUpgradeTool, Error, TEXT("Upgrade log parsing error: %s"), *UpgradeListResult.GetError());
+		UE_LOG(LogUnrealObjectPtrTool, Error, TEXT("Upgrade log parsing error: %s"), *UpgradeListResult.GetError());
 		return 1;
 	}
 
-	if (!NativePointerUpgrader.TryAllUpgrades(UpgradeListResult.GetValue()))
+	if (!PointerUpgrader.TryAllUpgrades(UpgradeListResult.GetValue()))
 	{
 		return 1;
 	}
 
-	if (EnumHasAllFlags(BehaviorFlags, ENativePointerUpgradeBehaviorFlags::PreviewOnly))
+	if (EnumHasAllFlags(BehaviorFlags, EPointerUpgradeBehaviorFlags::PreviewOnly))
 	{
-		UE_LOG(LogNativePointerUpgradeTool, Display, TEXT("Previewed upgrade successfully.  %d upgrades found; %d upgrades would be performed"), NativePointerUpgrader.GetTotalUpgradesFound(), NativePointerUpgrader.GetTotalUpgradesPerformed());
+		UE_LOG(LogUnrealObjectPtrTool, Display, TEXT("Previewed upgrade successfully.  %d upgrades found; %d upgrades would be performed"), PointerUpgrader.GetTotalUpgradesFound(), PointerUpgrader.GetTotalUpgradesPerformed());
 	}
 	else
 	{
-		UE_LOG(LogNativePointerUpgradeTool, Display, TEXT("Upgraded successfully.  %d upgrades found; %d upgrades performed"), NativePointerUpgrader.GetTotalUpgradesFound(), NativePointerUpgrader.GetTotalUpgradesPerformed());
+		UE_LOG(LogUnrealObjectPtrTool, Display, TEXT("Upgraded successfully.  %d upgrades found; %d upgrades performed"), PointerUpgrader.GetTotalUpgradesFound(), PointerUpgrader.GetTotalUpgradesPerformed());
 	}
 
 	return 0;
diff --git a/Engine/Source/Programs/UnrealObjectPtrTool/UnrealObjectPtrTool.Build.cs b/Engine/Source/Programs/UnrealObjectPtrTool/UnrealObjectPtrTool.Build.cs
--- a/Engine/Source/Programs/UnrealObjectPtrTool/UnrealObjectPtrTool.Build.cs
+++ b/Engine/Source/Programs/UnrealObjectPtrTool/UnrealObjectPtrTool.Build.cs
@@ -2,9 +2,9 @@
 
 using UnrealBuildTool;
 
-public class NativePointerUpgradeTool : ModuleRules
+public class UnrealObjectPtrTool : ModuleRules
 {
-	public NativePointerUpgradeTool(ReadOnlyTargetRules Target) : base(Target)
+	public UnrealObjectPtrTool(ReadOnlyTargetRules Target) : base(Target)
 	{
 		PublicIncludePaths.Add("Runtime/Launch/Public");
 
diff --git a/Engine/Source/Programs/UnrealObjectPtrTool/UnrealObjectPtrTool.Target.cs b/Engine/Source/Programs/UnrealObjectPtrTool/UnrealObjectPtrTool.Target.cs
--- a/Engine/Source/Programs/UnrealObjectPtrTool/UnrealObjectPtrTool.Target.cs
+++ b/Engine/Source/Programs/UnrealObjectPtrTool/UnrealObjectPtrTool.Target.cs
@@ -4,13 +4,13 @@ using UnrealBuildTool;
 using System.Collections.Generic;
 
 [SupportedPlatforms(UnrealPlatformClass.Desktop)]
-public class NativePointerUpgradeToolTarget : TargetRules
+public class UnrealObjectPtrToolTarget : TargetRules
 {
-	public NativePointerUpgradeToolTarget(TargetInfo Target) : base(Target)
+	public UnrealObjectPtrToolTarget(TargetInfo Target) : base(Target)
 	{
 		Type = TargetType.Program;
 		LinkType = TargetLinkType.Modular;
-		LaunchModuleName = "NativePointerUpgradeTool";
+		LaunchModuleName = "UnrealObjectPtrTool";
 
 		bBuildDeveloperTools			= false;
 		bUseMallocProfiler				= false;
