Prototype Wrapped Object Pointers (ObjectHandle/ObjectPtr)

Allow UnrealObjectPtrTool to perform reverse upgrades when the "-r" commandline argument is used.  Also added functionality to allow UnrealHeaderTool to dump out necessary log messages specifying TObjectPtr member declrations for UnrealObjectPtrTool to consume.

Tested by:
1) Performing a new forward pointer upgrade (from a new UHT log) and confirming the resulting files had zero diffs when compared to a shelf of upgraded files from previous version of UnrealObjectPtrTool.
2) Performed a forward pointer upgrade, then a reverse pointer upgrade then reverted all unchanged files.  The remaining files (13 of them see shelved CL15213431) only had trailing whitespace corrections due to the deliberate choices for trailing whitespace manipulation when doing forward upgrades.

diff --git a/Engine/Programs/UnrealHeaderTool/Config/DefaultEngine.ini b/Engine/Programs/UnrealHeaderTool/Config/DefaultEngine.ini
--- a/Engine/Programs/UnrealHeaderTool/Config/DefaultEngine.ini
+++ b/Engine/Programs/UnrealHeaderTool/Config/DefaultEngine.ini
@@ -35,3 +35,4 @@ gc.MaxObjectsInProgram=262144
 ; This will increase the size of the UHT log files, but the size should return
 ; to normal after native pointer upgrades are completed
 NativePointerMemberBehavior=AllowAndLog
+ObjectPtrMemberBehavior=AllowAndLog
diff --git a/Engine/Source/Programs/UnrealHeaderTool/Private/HeaderParser.cpp b/Engine/Source/Programs/UnrealHeaderTool/Private/HeaderParser.cpp
--- a/Engine/Source/Programs/UnrealHeaderTool/Private/HeaderParser.cpp
+++ b/Engine/Source/Programs/UnrealHeaderTool/Private/HeaderParser.cpp
@@ -107,7 +107,8 @@ const FName FHeaderParserNames::NAME_IsConversionRoot(TEXT("IsConversionRoot"));
 const FName FHeaderParserNames::NAME_AdvancedClassDisplay(TEXT("AdvancedClassDisplay"));
 
 EGeneratedCodeVersion FHeaderParser::DefaultGeneratedCodeVersion = EGeneratedCodeVersion::V1;
-ENativePointerMemberBehavior FHeaderParser::NativePointerMemberBehavior = ENativePointerMemberBehavior::AllowSilently;
+EPointerMemberBehavior FHeaderParser::NativePointerMemberBehavior = EPointerMemberBehavior::AllowSilently;
+EPointerMemberBehavior FHeaderParser::ObjectPtrMemberBehavior = EPointerMemberBehavior::AllowSilently;
 TArray<FString> FHeaderParser::StructsWithNoPrefix;
 TArray<FString> FHeaderParser::StructsWithTPrefix;
 FRigVMStructMap FHeaderParser::StructRigVMMap;
@@ -1303,25 +1304,25 @@ namespace
 		SkipAlignasIfNecessary(Parser);
 	}
 
-	inline ENativePointerMemberBehavior ToNativePointerMemberBehavior(const FString& InString)
+	inline EPointerMemberBehavior ToPointerMemberBehavior(const FString& InString)
 	{
 		if (InString.Compare(TEXT("Disallow")) == 0)
 		{
-			return ENativePointerMemberBehavior::Disallow;
+			return EPointerMemberBehavior::Disallow;
 		}
 
 		if (InString.Compare(TEXT("AllowSilently")) == 0)
 		{
-			return ENativePointerMemberBehavior::AllowSilently;
+			return EPointerMemberBehavior::AllowSilently;
 		}
 
 		if (InString.Compare(TEXT("AllowAndLog")) == 0)
 		{
-			return ENativePointerMemberBehavior::AllowAndLog;
+			return EPointerMemberBehavior::AllowAndLog;
 		}
 
 		FError::Throwf(TEXT("Unrecognized native pointer member behavior: %s"), *InString );
-		return ENativePointerMemberBehavior::Disallow;
+		return EPointerMemberBehavior::Disallow;
 	}
 
 	static const TCHAR* GLayoutMacroNames[] = {
@@ -4973,37 +4974,17 @@ void FHeaderParser::GetVarType(
 					// Optionally emit messages about native pointer members and swallow trailing 'const' after pointer properties
 					if (VariableCategory == EVariableCategory::Member)
 					{
-						switch (NativePointerMemberBehavior)
-						{
-							case ENativePointerMemberBehavior::Disallow:
-							{
-								FString TypeName = FString(InputPos - VarStartPos, Input + VarStartPos).TrimStartAndEnd().ReplaceCharWithEscapedChar();
-								UE_LOG(LogCompile, Error, TEXT("Native pointer usage in member declaration detected in '%s', line %d [[%s]]"), *GetSourceFileContext(), InputLine, *TypeName);
-							}
-							break;
-							case ENativePointerMemberBehavior::AllowAndLog:
-							{
-								FString TypeName = FString(InputPos - VarStartPos, Input + VarStartPos).TrimStartAndEnd().ReplaceCharWithEscapedChar();
-								UE_LOG(LogCompile, Log, TEXT("Native pointer usage in member declaration detected in '%s', line %d [[%s]]"), *GetSourceFileContext(), InputLine, *TypeName);
-							}
-							break;
-							case ENativePointerMemberBehavior::AllowSilently:
-							{
-								// Do nothing
-							}
-							break;
-							default:
-							{
-								checkf(false, TEXT("Unhandled case"));
-							}
-							break;
-						}
+						ConditionalLogPointerUsage(NativePointerMemberBehavior, TEXT("Native pointer"), FString(InputPos - VarStartPos, Input + VarStartPos).TrimStartAndEnd().ReplaceCharWithEscapedChar());
 
 						MatchIdentifier(TEXT("const"), ESearchCase::CaseSensitive);
 					}
 
 					VarProperty.PointerType = EPointerType::Native;
 				}
+				else if ((PropertyType == CPT_ObjectPtrReference) && (VariableCategory == EVariableCategory::Member))
+				{
+					ConditionalLogPointerUsage(ObjectPtrMemberBehavior, TEXT("ObjectPtr"), FString(InputPos - VarStartPos, Input + VarStartPos).TrimStartAndEnd().ReplaceCharWithEscapedChar());
+				}
 
 				// Imply const if it's a parameter that is a pointer to a const class
 				if (VariableCategory != EVariableCategory::Member && (TempClass != NULL) && (TempClass->HasAnyClassFlags(CLASS_Const)))
@@ -9019,6 +9000,7 @@ FHeaderParser::FHeaderParser(FFeedbackContext* InWarn, const FManifestModule& In
 		const FName DelegateParameterCountStringsKey(TEXT("DelegateParameterCountStrings"));
 		const FName GeneratedCodeVersionKey(TEXT("GeneratedCodeVersion"));
 		const FName NativePointerMemberBehaviorKey(TEXT("NativePointerMemberBehavior"));
+		const FName ObjectPtrMemberBehaviorKey(TEXT("ObjectPtrMemberBehavior"));
 
 		FConfigSection* ConfigSection = GConfig->GetSectionPrivate(TEXT("UnrealHeaderTool"), false, true, GEngineIni);
 		if (ConfigSection)
@@ -9053,7 +9035,11 @@ FHeaderParser::FHeaderParser(FFeedbackContext* InWarn, const FManifestModule& In
 				}
 				else if (It.Key() == NativePointerMemberBehaviorKey)
 				{
-					NativePointerMemberBehavior = ToNativePointerMemberBehavior(It.Value().GetValue());
+					NativePointerMemberBehavior = ToPointerMemberBehavior(It.Value().GetValue());
+				}
+				else if (It.Key() == ObjectPtrMemberBehaviorKey)
+				{
+					ObjectPtrMemberBehavior = ToPointerMemberBehavior(It.Value().GetValue());
 				}
 			}
 		}
@@ -10911,6 +10897,33 @@ bool FHeaderParser::CheckUIMinMaxRangeFromMetaData(const FString& UIMin, const F
 	return true;
 }
 
+void FHeaderParser::ConditionalLogPointerUsage(EPointerMemberBehavior PointerMemberBehavior, const TCHAR* PointerTypeDesc, FString&& PointerTypeDecl)
+{
+	switch (NativePointerMemberBehavior)
+	{
+		case EPointerMemberBehavior::Disallow:
+		{
+			UE_LOG(LogCompile, Error, TEXT("%s usage in member declaration detected in '%s', line %d [[%s]]"), PointerTypeDesc, *GetSourceFileContext(), InputLine, *PointerTypeDecl);
+		}
+		break;
+		case EPointerMemberBehavior::AllowAndLog:
+		{
+			UE_LOG(LogCompile, Log, TEXT("%s usage in member declaration detected in '%s', line %d [[%s]]"), PointerTypeDesc, *GetSourceFileContext(), InputLine, *PointerTypeDecl);
+		}
+		break;
+		case EPointerMemberBehavior::AllowSilently:
+		{
+			// Do nothing
+		}
+		break;
+		default:
+		{
+			checkf(false, TEXT("Unhandled case"));
+		}
+		break;
+	}
+}
+
 template <class TFunctionType>
 TFunctionType* CreateFunctionImpl(const FFuncInfo& FuncInfo, UObject* Outer, FScope* CurrentScope)
 {
diff --git a/Engine/Source/Programs/UnrealHeaderTool/Private/HeaderParser.h b/Engine/Source/Programs/UnrealHeaderTool/Private/HeaderParser.h
--- a/Engine/Source/Programs/UnrealHeaderTool/Private/HeaderParser.h
+++ b/Engine/Source/Programs/UnrealHeaderTool/Private/HeaderParser.h
@@ -419,7 +419,7 @@ struct FDocumentationPolicy
 	bool bFloatRangesRequired = UHT_DOCUMENTATION_POLICY_DEFAULT;
 };
 
-enum class ENativePointerMemberBehavior
+enum class EPointerMemberBehavior
 {
 	Disallow,
 	AllowSilently,
@@ -439,7 +439,8 @@ public:
 	// Default version of generated code. Defaults to oldest possible, unless specified otherwise in config.
 	static EGeneratedCodeVersion DefaultGeneratedCodeVersion;
 
-	static ENativePointerMemberBehavior NativePointerMemberBehavior;
+	static EPointerMemberBehavior NativePointerMemberBehavior;
+	static EPointerMemberBehavior ObjectPtrMemberBehavior;
 
 	// Compute the function parameter size and save the return offset
 	static void ComputeFunctionParametersSize(UClass* InClass);
@@ -1079,6 +1080,9 @@ private:
 	// Checks if a valid range has been found on the provided metadata
 	bool CheckUIMinMaxRangeFromMetaData(const FString& UIMin, const FString& UIMax);
 
+	// Checks if a valid range has been found on the provided metadata
+	void ConditionalLogPointerUsage(EPointerMemberBehavior PointerMemberBehavior, const TCHAR* PointerTypeDesc, FString&& PointerTypeDecl);
+
 	// Names that cannot be used enums, UStructs, or UClasses
 	static TArray<FString> ReservedTypeNames;
 
diff --git a/Engine/Source/Programs/UnrealObjectPtrTool/Private/UnrealObjectPtrTool.cpp b/Engine/Source/Programs/UnrealObjectPtrTool/Private/UnrealObjectPtrTool.cpp
--- a/Engine/Source/Programs/UnrealObjectPtrTool/Private/UnrealObjectPtrTool.cpp
+++ b/Engine/Source/Programs/UnrealObjectPtrTool/Private/UnrealObjectPtrTool.cpp
@@ -22,6 +22,7 @@ enum class EPointerUpgradeBehaviorFlags
 {
 	None,
 	PreviewOnly = 1,
+	Reverse     = 2,
 };
 ENUM_CLASS_FLAGS(EPointerUpgradeBehaviorFlags)
 
@@ -77,6 +78,8 @@ private:
 	static bool SaveStringArrayWithNewlinesToFile(const TArray<FString>& InStringArray, EByteOrderMarkerType InBOM, const TCHAR* InFilename, uint32 WriteFlags = 0);
 
 	static int32 CountLines(const FStringView Input);
+	static FString PerformForwardTypenameUpgrade(const FString& TypeName);
+	static FString PerformReverseTypenameUpgrade(const FString& TypeName);
 
 
 	int32 TotalUpgradesFound = 0;
@@ -105,45 +108,59 @@ FStringView ParseBracketedStringSegment(const FString& InString, int32& InOutCur
 TValueOrError<FPointerUpgrader::FPointerUpgradeList, FString> FPointerUpgrader::FillUpgradeList(const FString& LogFilename)
 {
 	FPointerUpgradeList UpgradeList;
-	const TCHAR Preamble[] = TEXT("Native pointer usage in member declaration detected in '");
+	const TCHAR* Preamble = EnumHasAllFlags(BehaviorFlags, EPointerUpgradeBehaviorFlags::Reverse) ? TEXT("ObjectPtr usage in member declaration detected in '") : TEXT("Native pointer usage in member declaration detected in '");
 	TArray<FString> PointerUpgradeEntries;
 	if (!FFileHelper::LoadFileToStringArrayWithPredicate(PointerUpgradeEntries, *LogFilename, [&Preamble](const FString& Line) { return Line.Contains(Preamble);  }))
 	{
 		return MakeError(FString::Printf(TEXT("Unable to load UHT log: %s"), *LogFilename));
 	}
 
-	constexpr int32 PreambleLen = UE_ARRAY_COUNT(Preamble) - 1;
+	const int32 PreambleLen = FCString::Strlen(Preamble);
 	for (const FString& Entry : PointerUpgradeEntries)
 	{
 		int32 CurrentIndex = 0;
 		FStringView FilenameView = ParseBracketedStringSegment(Entry, CurrentIndex, Preamble, TEXT("'"));
 		if (FilenameView.IsEmpty())
 		{
-			return MakeError(FString::Printf(TEXT("Failed to parse filename from native pointer member upgrade statement: %s"), *Entry));
+			return MakeError(FString::Printf(TEXT("Failed to parse filename from pointer member upgrade statement: %s"), *Entry));
 		}
 
 		FStringView LineNumberView = ParseBracketedStringSegment(Entry, CurrentIndex, TEXT(", line "), TEXT(" "));
 		if (LineNumberView.IsEmpty())
 		{
-			return MakeError(FString::Printf(TEXT("Failed to parse line number from native pointer member upgrade statement: %s"), *Entry));
+			return MakeError(FString::Printf(TEXT("Failed to parse line number from pointer member upgrade statement: %s"), *Entry));
 		}
 
 		FStringView TypeNameView = ParseBracketedStringSegment(Entry, CurrentIndex, TEXT("[["), TEXT("]]"));
 		if (TypeNameView.IsEmpty())
 		{
-			return MakeError(FString::Printf(TEXT("Failed to parse type name from native pointer member upgrade statement: %s"), *Entry));
+			return MakeError(FString::Printf(TEXT("Failed to parse type name from pointer member upgrade statement: %s"), *Entry));
 		}
 
 		int32 LineNumber = -1;
 		LexFromString(LineNumber, LineNumberView);
 		if (LineNumber < 1)
 		{
-			return MakeError(FString::Printf(TEXT("Failed to convert line number to integer from native pointer member upgrade statement: %s"), *Entry));
+			return MakeError(FString::Printf(TEXT("Failed to convert line number to integer from pointer member upgrade statement: %s"), *Entry));
 		}
 
-		if (!TypeNameView.EndsWith(TEXT("*")))
+		if (EnumHasAllFlags(BehaviorFlags, EPointerUpgradeBehaviorFlags::Reverse))
 		{
-			return MakeError(FString::Printf(TEXT("Failed to find expected asterisk at end of native pointer member upgrade statement: %s"), *Entry));
+			if (!TypeNameView.StartsWith(TEXT("TObjectPtr<")))
+			{
+				return MakeError(FString::Printf(TEXT("Failed to find expected 'TObjectPtr<' at end of pointer member upgrade statement: %s"), *Entry));
+			}
+			if (!TypeNameView.EndsWith(TEXT(">")))
+			{
+				return MakeError(FString::Printf(TEXT("Failed to find expected angle bracket at end of pointer member upgrade statement: %s"), *Entry));
+			}
+		}
+		else
+		{
+			if (!TypeNameView.EndsWith(TEXT("*")))
+			{
+				return MakeError(FString::Printf(TEXT("Failed to find expected asterisk at end of pointer member upgrade statement: %s"), *Entry));
+			}
 		}
 
 		TArray<FMemberDeclaration>& DeclarationArray = UpgradeList.FindOrAdd(FString(FilenameView));
@@ -196,25 +213,7 @@ bool FPointerUpgrader::TryUpgradesInFile(const FString& FilenameToUpgrade, TCons
 		}
 		FString OriginalLinesToUpgrade = LinesToUpgrade;
 
-		// Chop off the trailing asterisk in the member type declaration
-		FString TemplateArg = MemberDeclaration.TypeName.LeftChop(1);
-		// Any trailing whitespace that was after the asterisk should not go in the template argument, but be appended after the template argument
-		// This only applies to space (' ') and tab ('\t') that have non-whitespace before them on the line, and should not include newlines.
-		// This is to avoid having end template bracket ('>') get put on a line that may have a line comment ahead of it ("//").
-		int32 TrailingWhitespaceCount = 0;
-		int32 CharIndex;
-		for (CharIndex = TemplateArg.Len()-1; (CharIndex >= 0) && (TemplateArg[CharIndex] == TCHAR(' ') || TemplateArg[CharIndex] == TCHAR('\t')); --CharIndex)
-		{
-			++TrailingWhitespaceCount;
-		}
-		if ((CharIndex == 0) || (TemplateArg[CharIndex-1] == TEXT('\n')) || (TemplateArg[CharIndex-1] == TEXT('\r')))
-		{
-			// The trailing whitespace represented the entire lead of the line the end template character will be on.  Put the indentation inside the template arg instead of outside of it.
-			TrailingWhitespaceCount = 0;
-		}
-		FString TrailingWhitespace = TemplateArg.Right(TrailingWhitespaceCount);
-		TemplateArg.LeftChopInline(TrailingWhitespaceCount);
-		FString UpgradedTypeName = FString::Printf(TEXT("TObjectPtr<%s>%s"), *TemplateArg, *TrailingWhitespace);
+		FString UpgradedTypeName = EnumHasAllFlags(BehaviorFlags, EPointerUpgradeBehaviorFlags::Reverse) ? PerformReverseTypenameUpgrade(MemberDeclaration.TypeName) : PerformForwardTypenameUpgrade(MemberDeclaration.TypeName);
 		if (LinesToUpgrade.ReplaceInline(*MemberDeclaration.TypeName, *UpgradedTypeName) < 1)
 		{
 			if (!LinesToUpgrade.Contains(*UpgradedTypeName))
@@ -560,6 +559,39 @@ int32 FPointerUpgrader::CountLines(const FStringView Input)
 	return LineCount;
 }
 
+FString FPointerUpgrader::PerformForwardTypenameUpgrade(const FString& TypeName)
+{
+	// Chop off the trailing asterisk in the type declaration
+	FString TemplateArg = TypeName.LeftChop(1);
+	// Any trailing whitespace that was after the asterisk should not go in the template argument, but be appended after the template argument
+	// This only applies to space (' ') and tab ('\t') that have non-whitespace before them on the line, and should not include newlines.
+	// This is to avoid having end template bracket ('>') get put on a line that may have a line comment ahead of it ("//").
+	int32 TrailingWhitespaceCount = 0;
+	int32 CharIndex;
+	for (CharIndex = TemplateArg.Len()-1; (CharIndex >= 0) && (TemplateArg[CharIndex] == TCHAR(' ') || TemplateArg[CharIndex] == TCHAR('\t')); --CharIndex)
+	{
+		++TrailingWhitespaceCount;
+	}
+	if ((CharIndex == 0) || (TemplateArg[CharIndex-1] == TEXT('\n')) || (TemplateArg[CharIndex-1] == TEXT('\r')))
+	{
+		// The trailing whitespace represented the entire lead of the line the end template character will be on.  Put the indentation inside the template arg instead of outside of it.
+		TrailingWhitespaceCount = 0;
+	}
+	FString TrailingWhitespace = TemplateArg.Right(TrailingWhitespaceCount);
+	TemplateArg.LeftChopInline(TrailingWhitespaceCount);
+	return FString::Printf(TEXT("TObjectPtr<%s>%s"), *TemplateArg, *TrailingWhitespace);
+}
+
+FString FPointerUpgrader::PerformReverseTypenameUpgrade(const FString& TypeName)
+{
+	// Chop off the trailing angle bracket in the type declaration
+	FString TemplateArg = TypeName.LeftChop(1);
+	// Chop off the leading TObjectPtr in the type declaration
+	TemplateArg = TemplateArg.RightChop(UE_ARRAY_COUNT(TEXT("TObjectPtr<")) - 1);
+
+	return FString::Printf(TEXT("%s*"), *TemplateArg);
+}
+
 //////////////////////////////////////////////////////////////////////////
 
 INT32_MAIN_INT32_ARGC_TCHAR_ARGV()
@@ -593,6 +625,10 @@ INT32_MAIN_INT32_ARGC_TCHAR_ARGV()
 	{
 		BehaviorFlags |= EPointerUpgradeBehaviorFlags::PreviewOnly;
 	}
+	if (FParse::Param(*CmdLine, TEXT("r")) || FParse::Param(*CmdLine, TEXT("REVERSE")))
+	{
+		BehaviorFlags |= EPointerUpgradeBehaviorFlags::Reverse;
+	}
 
 	FPointerUpgrader PointerUpgrader(SCCCommand, BehaviorFlags);
 
