Don't limit UnrealObjectPtrTool to only consuming Info verbosity lines.

diff --git a/Engine/Source/Programs/UnrealObjectPtrTool/Private/UnrealObjectPtrTool.cpp b/Engine/Source/Programs/UnrealObjectPtrTool/Private/UnrealObjectPtrTool.cpp
--- a/Engine/Source/Programs/UnrealObjectPtrTool/Private/UnrealObjectPtrTool.cpp
+++ b/Engine/Source/Programs/UnrealObjectPtrTool/Private/UnrealObjectPtrTool.cpp
@@ -116,7 +116,7 @@ TValueOrError<FPointerUpgrader::FPointerUpgradeList, FString> FPointerUpgrader::
 	FPointerUpgradeList UpgradeList;
 	const TCHAR* Preamble = EnumHasAllFlags(BehaviorFlags, EPointerUpgradeBehaviorFlags::Reverse) ? TEXT("ObjectPtr usage in member declaration detected [[") : TEXT("Native pointer usage in member declaration detected [[");
 	TArray<FString> PointerUpgradeEntries;
-	if (!FFileHelper::LoadFileToStringArrayWithPredicate(PointerUpgradeEntries, *LogFilename, [&Preamble](const FString& Line) { return Line.Contains("Info: ") && Line.Contains(Preamble);  }))
+	if (!FFileHelper::LoadFileToStringArrayWithPredicate(PointerUpgradeEntries, *LogFilename, [&Preamble](const FString& Line) { return Line.Contains(Preamble);  }))
 	{
 		return MakeError(FString::Printf(TEXT("Unable to load UHT log: %s"), *LogFilename));
 	}
